Composure: Add ComposureShared plugin and migrate generic assets from legacy Composure to be shared with new Composure (Composite).

* Creates a new content-only ComposureShared plugin as a shared content container.
* Both Composure and Composite declare it as a plugin dependency so it is always mounted.
* Migrates 29 generic compositing assets from legacy Composure into it, with redirectors left in place for backwards compatibility.

#virtualized
#Synced-CL 52419094

diff --git a/Engine/Plugins/Compositing/Composite/Composite.uplugin b/Engine/Plugins/Compositing/Composite/Composite.uplugin
--- a/Engine/Plugins/Compositing/Composite/Composite.uplugin
+++ b/Engine/Plugins/Compositing/Composite/Composite.uplugin
@@ -32,6 +32,10 @@
 			"Name": "CompositeCore",
 			"Enabled": true
 		},
+		{
+			"Name": "ComposureShared",
+			"Enabled": true
+		},
 		{
 			"Name": "OpenColorIO",
 			"Enabled": true
diff --git a/Engine/Plugins/Compositing/Composure/Composure.uplugin b/Engine/Plugins/Compositing/Composure/Composure.uplugin
--- a/Engine/Plugins/Compositing/Composure/Composure.uplugin
+++ b/Engine/Plugins/Compositing/Composure/Composure.uplugin
@@ -36,6 +36,10 @@
 		}
 	],
 	"Plugins": [
+		{
+			"Name": "ComposureShared",
+			"Enabled": true
+		},
 		{
 			"Name": "CameraCalibrationCore",
 			"Enabled": true
diff --git a/Engine/Plugins/Compositing/ComposureShared/ComposureShared.uplugin b/Engine/Plugins/Compositing/ComposureShared/ComposureShared.uplugin
new file mode 100644
--- /dev/null
+++ b/Engine/Plugins/Compositing/ComposureShared/ComposureShared.uplugin
@@ -0,0 +1,18 @@
+{
+	"FileVersion": 1,
+	"Version": 1,
+	"VersionName": "1.0",
+	"FriendlyName": "Composure Shared",
+	"Description": "Shared content assets (materials, textures, meshes) migrated from legacy Composure, accessible to both Composure and Composite plugins.",
+	"Category": "Compositing",
+	"CreatedBy": "Epic Games, Inc.",
+	"CreatedByURL": "https://epicgames.com",
+	"DocsURL": "",
+	"MarketplaceURL": "",
+	"SupportURL": "",
+	"EnabledByDefault": false,
+	"CanContainContent": true,
+	"IsBetaVersion": false,
+	"IsExperimentalVersion": false,
+	"Installed": false
+}
